- GearVR: fixed a crash when stopping and starting the app (proximity sensor app stop)
- GearVR: Bring CPU and GPU levels in from engine.ini

diff --git a/Engine/Plugins/Runtime/GearVR/Source/GearVR/Private/GearVR.cpp b/Engine/Plugins/Runtime/GearVR/Source/GearVR/Private/GearVR.cpp
--- a/Engine/Plugins/Runtime/GearVR/Source/GearVR/Private/GearVR.cpp
+++ b/Engine/Plugins/Runtime/GearVR/Source/GearVR/Private/GearVR.cpp
@@ -980,13 +980,22 @@ bool FGearVR::IsInitialized() const
 
 void FGearVR::Startup()
 {
+	// grab the clock settings out of the ini
+	const TCHAR* GearVRSettings = TEXT("GearVR.Settings");
+	int CpuLevel = 2;
+	int GpuLevel = 2;
+	GConfig->GetInt(GearVRSettings, TEXT("CpuLevel"), CpuLevel, GEngineIni);
+	GConfig->GetInt(GearVRSettings, TEXT("GpuLevel"), GpuLevel, GEngineIni);
+
+	UE_LOG(LogHMD, Log, TEXT("GearVR starting with CPU: %d GPU: %d"), CpuLevel, GpuLevel);
+
 	FMemory::MemZero(VrModeParms);
 	VrModeParms.AsynchronousTimeWarp = true;
 	VrModeParms.DistortionFileName = NULL;
 	VrModeParms.EnableImageServer = false;
 	VrModeParms.GameThreadTid = gettid();
-	VrModeParms.CpuLevel = 2;
-	VrModeParms.GpuLevel = 2;
+	VrModeParms.CpuLevel = CpuLevel;
+	VrModeParms.GpuLevel = GpuLevel;
 	VrModeParms.ActivityObject = FJavaWrapper::GameActivityThis;
 
 	FPlatformMisc::MemoryBarrier();
@@ -1361,20 +1370,20 @@ void FGearVR::ShutdownRendering()
 	if (OvrMobile)
 	{
 		ovr_LeaveVrMode(OvrMobile);
+		OvrMobile = nullptr;
+
+		check(GJavaVM);
+		const jint DetachResult = GJavaVM->DetachCurrentThread();
+		if (DetachResult == JNI_ERR)
+		{
+			FPlatformMisc::LowLevelOutputDebugString(TEXT("FJNIHelper failed to detach thread from Java VM!"));
+		}
 	}
 
 	if (pGearVRBridge)
 	{
 		pGearVRBridge->Shutdown();
-		pGearVRBridge = NULL;
-	}
-
-	check(GJavaVM);
-	const jint DetachResult = GJavaVM->DetachCurrentThread();
-	if (DetachResult == JNI_ERR)
-	{
-		FPlatformMisc::LowLevelOutputDebugString(TEXT("FJNIHelper failed to detach thread from Java VM!"));
-		check(false);
+		pGearVRBridge = nullptr;
 	}
 }
 
@@ -1566,7 +1575,9 @@ void FGearVR::FGearVRBridge::UpdateViewport(const FViewport& Viewport, FRHIViewp
 	const uint32 RTSizeY = RT->GetSizeY();
 	GLuint RTTexId = *(GLuint*)RT->GetNativeResource();
 
- 	const Matrix4f proj = Matrix4f::PerspectiveRH( FOV, 1.0f, 1.0f, 100.0f);
+	FMatrix ProjMat = Plugin->GetStereoProjectionMatrix(eSSP_LEFT_EYE, 90.0f);
+	const Matrix4f proj = Plugin->ToMatrix4f(ProjMat);
+ //	const Matrix4f proj = Matrix4f::PerspectiveRH( FOV, 1.0f, 1.0f, 100.0f);
  	SwapParms.Images[0][0].TexCoordsFromTanAngles = TanAngleMatrixFromProjection( proj );
 // 	SwapParms.Images[0][0].TexId = RTTexId;
 // 	SwapParms.Images[0][0].Pose = Plugin->RenderParams_RenderThread.EyeRenderPose[0];
diff --git a/Engine/Plugins/Runtime/GearVR/Source/GearVR/Private/GearVR.h b/Engine/Plugins/Runtime/GearVR/Source/GearVR/Private/GearVR.h
--- a/Engine/Plugins/Runtime/GearVR/Source/GearVR/Private/GearVR.h
+++ b/Engine/Plugins/Runtime/GearVR/Source/GearVR/Private/GearVR.h
@@ -265,6 +265,14 @@ private:
 			FPlane(vtm.M[0][3], vtm.M[1][3], vtm.M[2][3], vtm.M[3][3]));
 	}
 
+	FORCEINLINE OVR::Matrix4f ToMatrix4f(const FMatrix mat) const
+	{
+		return OVR::Matrix4f(mat.M[0][0],mat.M[1][0],mat.M[2][0],mat.M[3][0],
+							 mat.M[0][1],mat.M[1][1],mat.M[2][1],mat.M[3][1],
+							 mat.M[0][2],mat.M[1][2],mat.M[2][2],mat.M[3][2],
+							 mat.M[0][3],mat.M[1][3],mat.M[2][3],mat.M[3][3]);
+	}
+
 	/**
 	 * Called when state changes from 'stereo' to 'non-stereo'. Suppose to distribute
 	 * the event further to user's code (?).
