- Fix for GearVR global menu call.

diff --git a/Engine/Plugins/Runtime/GearVR/Source/GearVR/Private/GearVR.cpp b/Engine/Plugins/Runtime/GearVR/Source/GearVR/Private/GearVR.cpp
--- a/Engine/Plugins/Runtime/GearVR/Source/GearVR/Private/GearVR.cpp
+++ b/Engine/Plugins/Runtime/GearVR/Source/GearVR/Private/GearVR.cpp
@@ -580,7 +580,12 @@ bool FGearVR::Exec( UWorld* InWorld, const TCHAR* Cmd, FOutputDevice& Ar )
 	}
 	else if (FParse::Command(&Cmd, TEXT("OVRGLOBALMENU")))
 	{
-		ovr_StartSystemActivity(OvrMobile, PUI_GLOBAL_MENU, NULL);
+		// fire off the global menu from the render thread
+		ENQUEUE_UNIQUE_RENDER_COMMAND_ONEPARAMETER(OVRGlobalMenu,
+			FGearVR*, Plugin, this,
+			{
+				Plugin->StartOVRGlobalMenu();
+			});
 	}
 
 	return false;
@@ -1392,6 +1397,16 @@ void FGearVR::ShutdownRendering()
 	}
 }
 
+void FGearVR::StartOVRGlobalMenu()
+{
+	check(IsInRenderingThread());
+
+	if (OvrMobile)
+	{
+		ovr_StartSystemActivity(OvrMobile, PUI_GLOBAL_MENU, NULL);
+	}
+}
+
 void FGearVR::UpdateViewport(bool bUseSeparateRenderTarget, const FViewport& InViewport, SViewport* ViewportWidget)
 {
 	check(IsInGameThread());
diff --git a/Engine/Plugins/Runtime/GearVR/Source/GearVR/Private/GearVR.h b/Engine/Plugins/Runtime/GearVR/Source/GearVR/Private/GearVR.h
--- a/Engine/Plugins/Runtime/GearVR/Source/GearVR/Private/GearVR.h
+++ b/Engine/Plugins/Runtime/GearVR/Source/GearVR/Private/GearVR.h
@@ -171,6 +171,7 @@ public:
 	TRefCountPtr<FGearVRBridge> pGearVRBridge;
 
 	void ShutdownRendering();
+	void StartOVRGlobalMenu();
 
 private:
 	FGearVR* getThis() { return this; }
