- Basic GearVR head model.
- Defaults to on.
- Scale/disable via engine config:

[GearVR.Settings]
HeadModelScale=0.0

diff --git a/Engine/Plugins/Runtime/GearVR/Source/GearVR/Private/GearVR.cpp b/Engine/Plugins/Runtime/GearVR/Source/GearVR/Private/GearVR.cpp
--- a/Engine/Plugins/Runtime/GearVR/Source/GearVR/Private/GearVR.cpp
+++ b/Engine/Plugins/Runtime/GearVR/Source/GearVR/Private/GearVR.cpp
@@ -740,10 +740,12 @@ void FGearVR::CalculateStereoViewOffset(const EStereoscopicPass StereoPassType,
 		check(WorldToMeters != 0.f)
 
 		const int idx = (StereoPassType == eSSP_LEFT_EYE) ? 0 : 1;
-		float EyeMul = (StereoPassType == eSSP_LEFT_EYE) ? -1.0f : 1.0f;
-		const float PassEyeOffset = InterpupillaryDistance * EyeMul * 0.5f * WorldToMeters;
+		float EyeMul = (StereoPassType == eSSP_LEFT_EYE) ? -0.5f : 0.5f;
 
-		const FVector TotalOffset = FVector(0, PassEyeOffset, 0);
+		FVector TotalOffset = FVector(0, InterpupillaryDistance * EyeMul, 0);
+		TotalOffset += HeadModel;
+
+		TotalOffset *= WorldToMeters;
 
 		ViewLocation += ViewRotation.Quaternion().RotateVector(TotalOffset);
 
@@ -946,6 +948,7 @@ FGearVR::FGearVR()
 	, IdealScreenPercentage(100.0f)
 	, bAllowFinishCurrentFrame(false)
 	, InterpupillaryDistance(OVR_DEFAULT_IPD)
+	, HeadModel(0.12f, 0.0f, 0.17f)
 	, WorldToMetersScale(100.f)
 	, bWorldToMetersOverride(false)
 	, UserDistanceToScreenModifier(0.f)
@@ -995,8 +998,10 @@ void FGearVR::Startup()
 	const TCHAR* GearVRSettings = TEXT("GearVR.Settings");
 	int CpuLevel = 2;
 	int GpuLevel = 2;
+	float HeadModelScale = 1.0f;
 	GConfig->GetInt(GearVRSettings, TEXT("CpuLevel"), CpuLevel, GEngineIni);
 	GConfig->GetInt(GearVRSettings, TEXT("GpuLevel"), GpuLevel, GEngineIni);
+	GConfig->GetFloat(GearVRSettings, TEXT("HeadModelScale"), HeadModelScale, GEngineIni);
 
 	UE_LOG(LogHMD, Log, TEXT("GearVR starting with CPU: %d GPU: %d"), CpuLevel, GpuLevel);
 
@@ -1010,6 +1015,8 @@ void FGearVR::Startup()
 	VrModeParms.GpuLevel = GpuLevel;
 	VrModeParms.ActivityObject = FJavaWrapper::GameActivityThis;
 
+	HeadModel *= HeadModelScale;
+
 	FPlatformMisc::MemoryBarrier();
 
 	if (!IsRunningGame() || (InitStatus & eStartupExecuted) != 0)
diff --git a/Engine/Plugins/Runtime/GearVR/Source/GearVR/Private/GearVR.h b/Engine/Plugins/Runtime/GearVR/Source/GearVR/Private/GearVR.h
--- a/Engine/Plugins/Runtime/GearVR/Source/GearVR/Private/GearVR.h
+++ b/Engine/Plugins/Runtime/GearVR/Source/GearVR/Private/GearVR.h
@@ -353,6 +353,9 @@ private: // data
 	/** Interpupillary distance, in meters (user configurable) */
 	float InterpupillaryDistance;
 
+	/** Vector defining center eye offset for head neck model in meters */
+	FVector HeadModel;
+
 	/** World units (UU) to Meters scale.  Read from the level, and used to transform positional tracking data */
 	float WorldToMetersScale;
 	/** Whether world-to-meters scale is overriden or not. */
