- Exposed GearVR MinimumVsyncs timewarp param via .ini
- Defaults to 1 (60Hz)

[GearVR.Settings]
MinimumVsyncs=2

diff --git a/Engine/Plugins/Runtime/GearVR/Source/GearVR/Private/GearVR.cpp b/Engine/Plugins/Runtime/GearVR/Source/GearVR/Private/GearVR.cpp
--- a/Engine/Plugins/Runtime/GearVR/Source/GearVR/Private/GearVR.cpp
+++ b/Engine/Plugins/Runtime/GearVR/Source/GearVR/Private/GearVR.cpp
@@ -956,6 +956,7 @@ FGearVR::FGearVR()
 	, VFOVInRadians(FMath::DegreesToRadians(90.f))
 	, RenderTargetWidth(2048)
 	, RenderTargetHeight(1024)
+	, MinimumVsyncs(1)
 	, MotionPredictionInSeconds(DEFAULT_PREDICTION_IN_SECONDS)
 	, bChromaAbCorrectionEnabled(true)
 	, bOverride2D(false)
@@ -1001,9 +1002,10 @@ void FGearVR::Startup()
 	float HeadModelScale = 1.0f;
 	GConfig->GetInt(GearVRSettings, TEXT("CpuLevel"), CpuLevel, GEngineIni);
 	GConfig->GetInt(GearVRSettings, TEXT("GpuLevel"), GpuLevel, GEngineIni);
+	GConfig->GetInt(GearVRSettings, TEXT("MinimumVsyncs"), MinimumVsyncs, GEngineIni);
 	GConfig->GetFloat(GearVRSettings, TEXT("HeadModelScale"), HeadModelScale, GEngineIni);
 
-	UE_LOG(LogHMD, Log, TEXT("GearVR starting with CPU: %d GPU: %d"), CpuLevel, GpuLevel);
+	UE_LOG(LogHMD, Log, TEXT("GearVR starting with CPU: %d GPU: %d MinimumVsyncs: %d"), CpuLevel, GpuLevel, MinimumVsyncs);
 
 	FMemory::Memzero(VrModeParms);
 	VrModeParms.SkipWindowFullscreenReset = true;
@@ -1038,7 +1040,7 @@ void FGearVR::Startup()
 	// Uncap fps to enable FPS higher than 62
 	GEngine->bSmoothFrameRate = false;
 
-	pGearVRBridge = new FGearVRBridge(this, RenderTargetWidth, RenderTargetHeight, HFOVInRadians);
+	pGearVRBridge = new FGearVRBridge(this, RenderTargetWidth, RenderTargetHeight, HFOVInRadians, MinimumVsyncs);
 
 	LoadFromIni();
 	SaveSystemValues();
@@ -1083,7 +1085,7 @@ void FGearVR::ApplicationResumeDelegate()
 	FPlatformMisc::LowLevelOutputDebugString(TEXT("+++++++ GEARVR APP RESUME ++++++"));
 	if(!pGearVRBridge)
 	{
-		pGearVRBridge = new FGearVRBridge(this, RenderTargetWidth, RenderTargetHeight, HFOVInRadians);
+		pGearVRBridge = new FGearVRBridge(this, RenderTargetWidth, RenderTargetHeight, HFOVInRadians, MinimumVsyncs);
 	}
 }
 
@@ -1451,12 +1453,13 @@ void FGearVR::UpdateViewport(bool bUseSeparateRenderTarget, const FViewport& InV
 }
 
 
-FGearVR::FGearVRBridge::FGearVRBridge(FGearVR* plugin, uint32 RenderTargetWidth, uint32 RenderTargetHeight, float FOV) :
+FGearVR::FGearVRBridge::FGearVRBridge(FGearVR* plugin, uint32 RenderTargetWidth, uint32 RenderTargetHeight, float FOV, int MinimumVsyncs) :
 	FRHICustomPresent(nullptr),
 	Plugin(plugin), 
 	bInitialized(false),
 	RenderTargetWidth(RenderTargetWidth),
 	RenderTargetHeight(RenderTargetHeight),
+	MinimumVsyncs(MinimumVsyncs),
 	FOV(FOV)
 {
 	Init();
@@ -1573,6 +1576,7 @@ void FGearVR::FGearVRBridge::Init()
 	}
 
 	SwapParms = InitTimeWarpParms();
+	SwapParms.MinimumVsyncs = MinimumVsyncs;
 }
 
 void FGearVR::FGearVRBridge::Reset()
diff --git a/Engine/Plugins/Runtime/GearVR/Source/GearVR/Private/GearVR.h b/Engine/Plugins/Runtime/GearVR/Source/GearVR/Private/GearVR.h
--- a/Engine/Plugins/Runtime/GearVR/Source/GearVR/Private/GearVR.h
+++ b/Engine/Plugins/Runtime/GearVR/Source/GearVR/Private/GearVR.h
@@ -129,7 +129,7 @@ public:
 	class FGearVRBridge : public FRHICustomPresent
 	{
 	public:
-		FGearVRBridge(FGearVR* plugin, uint32 RenderTargetWidth, uint32 RenderTargetHeight, float FOV );
+		FGearVRBridge(FGearVR* plugin, uint32 RenderTargetWidth, uint32 RenderTargetHeight, float FOV, int32 MinimumVsyncs);
 
 		// Returns true if it is initialized and used.
 		bool IsInitialized() const { return bInitialized; }
@@ -165,6 +165,7 @@ public:
 
 		uint32				RenderTargetWidth;
 		uint32				RenderTargetHeight;
+		int32				MinimumVsyncs;
 		float				FOV;
 	};
 
@@ -372,6 +373,9 @@ private: // data
 	int32 RenderTargetWidth;
 	int32 RenderTargetHeight;
 
+	/** Clamp warpswap to once every N vsyncs.  1 = 60Hz, 2 = 30Hz, etc. */
+	int32 MinimumVsyncs;
+
 	/** Motion prediction (in seconds). 0 - no prediction */
 	double MotionPredictionInSeconds;
 
