Buildwatcher - Fix shadow variables

diff --git a/Engine/Plugins/Runtime/GearVR/Source/GearVR/Private/GearVR.cpp b/Engine/Plugins/Runtime/GearVR/Source/GearVR/Private/GearVR.cpp
--- a/Engine/Plugins/Runtime/GearVR/Source/GearVR/Private/GearVR.cpp
+++ b/Engine/Plugins/Runtime/GearVR/Source/GearVR/Private/GearVR.cpp
@@ -1084,10 +1084,10 @@ void FGearVR::DrawDebug(UCanvas* Canvas)
 }
 
 //////////////////////////////////////////////////////////////////////////
-FGearVRCustomPresent::FGearVRCustomPresent(int MinimumVsyncs) :
+FGearVRCustomPresent::FGearVRCustomPresent(int InMinimumVsyncs) :
 	FRHICustomPresent(nullptr),
 	bInitialized(false),
-	MinimumVsyncs(MinimumVsyncs)
+	MinimumVsyncs(InMinimumVsyncs)
 {
 	Init();
 }
diff --git a/Engine/Plugins/Runtime/GearVR/Source/GearVR/Private/GearVR.h b/Engine/Plugins/Runtime/GearVR/Source/GearVR/Private/GearVR.h
--- a/Engine/Plugins/Runtime/GearVR/Source/GearVR/Private/GearVR.h
+++ b/Engine/Plugins/Runtime/GearVR/Source/GearVR/Private/GearVR.h
@@ -16,10 +16,12 @@
 	
 #include "OVRVersion.h"
 
+PRAGMA_DISABLE_SHADOW_VARIABLE_WARNINGS
 #include "../Src/Kernel/OVR_Math.h"
 #include "../Src/Kernel/OVR_Threads.h"
 #include "../Src/Kernel/OVR_Color.h"
 #include "../Src/Kernel/OVR_Timer.h"
+PRAGMA_ENABLE_SHADOW_VARIABLE_WARNINGS
 
 #include "OVR.h"
 #include "VrApi.h"
