INTEGRATE: Fixed stereo not working in VR PIE when Blueprint editor was open (UE-16835)
- Engine called OnStartGameFrame() with GWorld set to PIE world, even for non-PIE worlds
- Fixed by eliminating use of GWorld in HMD code

[Integrated from CL 2620236.  Reviewed by [REDACTED_PERSON]]

diff --git a/Engine/Plugins/Runtime/GearVR/Source/GearVR/Private/GearVR.cpp b/Engine/Plugins/Runtime/GearVR/Source/GearVR/Private/GearVR.cpp
--- a/Engine/Plugins/Runtime/GearVR/Source/GearVR/Private/GearVR.cpp
+++ b/Engine/Plugins/Runtime/GearVR/Source/GearVR/Private/GearVR.cpp
@@ -165,9 +165,9 @@ TSharedPtr<FHMDSettings, ESPMode::ThreadSafe> FGearVR::CreateNewSettings() const
 	return Result;
 }
 
-bool FGearVR::OnStartGameFrame()
+bool FGearVR::OnStartGameFrame( FWorldContext& WorldContext )
 {
-	bool rv = FHeadMountedDisplay::OnStartGameFrame();
+	bool rv = FHeadMountedDisplay::OnStartGameFrame( WorldContext );
 	if (!rv)
 	{
 		return false;
diff --git a/Engine/Plugins/Runtime/GearVR/Source/GearVR/Private/GearVR.h b/Engine/Plugins/Runtime/GearVR/Source/GearVR/Private/GearVR.h
--- a/Engine/Plugins/Runtime/GearVR/Source/GearVR/Private/GearVR.h
+++ b/Engine/Plugins/Runtime/GearVR/Source/GearVR/Private/GearVR.h
@@ -237,7 +237,7 @@ public:
 	static void PreInit();
 
 	/** IHeadMountedDisplay interface */
-	virtual bool OnStartGameFrame() override;
+	virtual bool OnStartGameFrame( FWorldContext& WorldContext ) override;
 	virtual bool IsHMDConnected() override { return true; }
 	virtual EHMDDeviceType::Type GetHMDDeviceType() const override;
 	virtual bool GetHMDMonitorInfo(MonitorInfo&) override;
diff --git a/Engine/Plugins/Runtime/GearVR/Source/GearVR/Private/HeadMountedDisplayCommon.cpp b/Engine/Plugins/Runtime/GearVR/Source/GearVR/Private/HeadMountedDisplayCommon.cpp
--- a/Engine/Plugins/Runtime/GearVR/Source/GearVR/Private/HeadMountedDisplayCommon.cpp
+++ b/Engine/Plugins/Runtime/GearVR/Source/GearVR/Private/HeadMountedDisplayCommon.cpp
@@ -195,14 +195,14 @@ FHMDGameFrame* FHeadMountedDisplay::GetCurrentFrame() const
 	return nullptr;
 }
 
-bool FHeadMountedDisplay::OnStartGameFrame()
+bool FHeadMountedDisplay::OnStartGameFrame( FWorldContext& WorldContext )
 {
 	check(IsInGameThread());
 
 	Frame.Reset();
 	Flags.bFrameStarted = true;
 
-	if (!GWorld || !GWorld->IsGameWorld())
+	if (!WorldContext.World() || !WorldContext.World()->IsGameWorld())
 	{
 		// ignore all non-game worlds
 		return false;
@@ -251,20 +251,18 @@ bool FHeadMountedDisplay::OnStartGameFrame()
 	CurrentFrame->FrameNumber = GFrameCounter;
 	CurrentFrame->Flags.bOutOfFrame = false;
 
-	check(GWorld);
-
 	if (Settings->Flags.bWorldToMetersOverride)
 	{
 		CurrentFrame->WorldToMetersScale = Settings->WorldToMetersScale;
 	}
 	else
 	{
-		CurrentFrame->WorldToMetersScale = GWorld->GetWorldSettings()->WorldToMeters;
+		CurrentFrame->WorldToMetersScale = WorldContext.World()->GetWorldSettings()->WorldToMeters;
 	}
 	return true;
 }
 
-bool FHeadMountedDisplay::OnEndGameFrame()
+bool FHeadMountedDisplay::OnEndGameFrame( FWorldContext& WorldContext )
 {
 	check(IsInGameThread());
 
@@ -272,7 +270,7 @@ bool FHeadMountedDisplay::OnEndGameFrame()
 
 	FHMDGameFrame* const CurrentGameFrame = Frame.Get();
 
-		if (!GWorld || !GWorld->IsGameWorld() || !CurrentGameFrame)
+	if (!WorldContext.World() || !WorldContext.World()->IsGameWorld() || !CurrentGameFrame)
 	{
 		// ignore all non-game worlds
 		return false;
diff --git a/Engine/Plugins/Runtime/GearVR/Source/GearVR/Private/HeadMountedDisplayCommon.h b/Engine/Plugins/Runtime/GearVR/Source/GearVR/Private/HeadMountedDisplayCommon.h
--- a/Engine/Plugins/Runtime/GearVR/Source/GearVR/Private/HeadMountedDisplayCommon.h
+++ b/Engine/Plugins/Runtime/GearVR/Source/GearVR/Private/HeadMountedDisplayCommon.h
@@ -296,12 +296,12 @@ public:
 	/**
 	 * This method is called when new game frame begins (called on a game thread).
 	 */
-	virtual bool OnStartGameFrame() override;
+	virtual bool OnStartGameFrame( FWorldContext& WorldContext ) override;
 
 	/**
 	 * This method is called when game frame ends (called on a game thread).
 	 */
-	virtual bool OnEndGameFrame() override;
+	virtual bool OnEndGameFrame( FWorldContext& WorldContext ) override;
 
 	virtual void SetupViewFamily(class FSceneViewFamily& InViewFamily) {}
 	virtual void SetupView(class FSceneViewFamily& InViewFamily, FSceneView& InView) {}
diff --git a/Engine/Plugins/Runtime/OculusRift/Source/OculusRift/Private/0.5/OculusRiftHMD_05.cpp b/Engine/Plugins/Runtime/OculusRift/Source/OculusRift/Private/0.5/OculusRiftHMD_05.cpp
--- a/Engine/Plugins/Runtime/OculusRift/Source/OculusRift/Private/0.5/OculusRiftHMD_05.cpp
+++ b/Engine/Plugins/Runtime/OculusRift/Source/OculusRift/Private/0.5/OculusRiftHMD_05.cpp
@@ -138,9 +138,9 @@ TSharedPtr<FHMDSettings, ESPMode::ThreadSafe> FOculusRiftHMD::CreateNewSettings(
 	return Result;
 }
 
-bool FOculusRiftHMD::OnStartGameFrame()
+bool FOculusRiftHMD::OnStartGameFrame( FWorldContext& WorldContext )
 {
-	bool rv = FHeadMountedDisplay::OnStartGameFrame();
+	bool rv = FHeadMountedDisplay::OnStartGameFrame( WorldContext );
 	if (!rv)
 	{
 		return false;
diff --git a/Engine/Plugins/Runtime/OculusRift/Source/OculusRift/Private/0.5/OculusRiftHMD_05.h b/Engine/Plugins/Runtime/OculusRift/Source/OculusRift/Private/0.5/OculusRiftHMD_05.h
--- a/Engine/Plugins/Runtime/OculusRift/Source/OculusRift/Private/0.5/OculusRiftHMD_05.h
+++ b/Engine/Plugins/Runtime/OculusRift/Source/OculusRift/Private/0.5/OculusRiftHMD_05.h
@@ -273,7 +273,7 @@ public:
 	static void PreInit();
 
 	/** IHeadMountedDisplay interface */
-	virtual bool OnStartGameFrame() override;
+	virtual bool OnStartGameFrame( FWorldContext& WorldContext ) override;
 
 	virtual bool IsHMDConnected() override;
 	virtual EHMDDeviceType::Type GetHMDDeviceType() const override;
diff --git a/Engine/Plugins/Runtime/OculusRift/Source/OculusRift/Private/HeadMountedDisplayCommon.cpp b/Engine/Plugins/Runtime/OculusRift/Source/OculusRift/Private/HeadMountedDisplayCommon.cpp
--- a/Engine/Plugins/Runtime/OculusRift/Source/OculusRift/Private/HeadMountedDisplayCommon.cpp
+++ b/Engine/Plugins/Runtime/OculusRift/Source/OculusRift/Private/HeadMountedDisplayCommon.cpp
@@ -197,11 +197,11 @@ FHMDGameFrame* FHeadMountedDisplay::GetCurrentFrame() const
 	return nullptr;
 }
 
-bool FHeadMountedDisplay::OnStartGameFrame()
+bool FHeadMountedDisplay::OnStartGameFrame( FWorldContext& WorldContext )
 {
 	check(IsInGameThread());
 
-	if (!GWorld || !GWorld->IsGameWorld())
+	if (!WorldContext.World() || !WorldContext.World()->IsGameWorld())
 	{
 		// ignore all non-game worlds
 		return false;
@@ -253,20 +253,18 @@ bool FHeadMountedDisplay::OnStartGameFrame()
 	CurrentFrame->FrameNumber = GFrameCounter;
 	CurrentFrame->Flags.bOutOfFrame = false;
 
-	check(GWorld);
-
 	if (Settings->Flags.bWorldToMetersOverride)
 	{
 		CurrentFrame->WorldToMetersScale = Settings->WorldToMetersScale;
 	}
 	else
 	{
-		CurrentFrame->WorldToMetersScale = GWorld->GetWorldSettings()->WorldToMeters;
+		CurrentFrame->WorldToMetersScale = WorldContext.World()->GetWorldSettings()->WorldToMeters;
 	}
 	return true;
 }
 
-bool FHeadMountedDisplay::OnEndGameFrame()
+bool FHeadMountedDisplay::OnEndGameFrame( FWorldContext& WorldContext )
 {
 	check(IsInGameThread());
 
@@ -274,7 +272,7 @@ bool FHeadMountedDisplay::OnEndGameFrame()
 
 	FHMDGameFrame* const CurrentGameFrame = Frame.Get();
 
-		if (!GWorld || !GWorld->IsGameWorld() || !CurrentGameFrame)
+	if (!WorldContext.World() || !WorldContext.World()->IsGameWorld() || !CurrentGameFrame)
 	{
 		// ignore all non-game worlds
 		return false;
diff --git a/Engine/Plugins/Runtime/OculusRift/Source/OculusRift/Private/HeadMountedDisplayCommon.h b/Engine/Plugins/Runtime/OculusRift/Source/OculusRift/Private/HeadMountedDisplayCommon.h
--- a/Engine/Plugins/Runtime/OculusRift/Source/OculusRift/Private/HeadMountedDisplayCommon.h
+++ b/Engine/Plugins/Runtime/OculusRift/Source/OculusRift/Private/HeadMountedDisplayCommon.h
@@ -295,12 +295,12 @@ public:
 	/**
 	 * This method is called when new game frame begins (called on a game thread).
 	 */
-	virtual bool OnStartGameFrame() override;
+	virtual bool OnStartGameFrame( FWorldContext& WorldContext ) override;
 
 	/**
 	 * This method is called when game frame ends (called on a game thread).
 	 */
-	virtual bool OnEndGameFrame() override;
+	virtual bool OnEndGameFrame( FWorldContext& WorldContext ) override;
 
 	virtual void SetupViewFamily(class FSceneViewFamily& InViewFamily) {}
 	virtual void SetupView(class FSceneViewFamily& InViewFamily, FSceneView& InView) {}
diff --git a/Engine/Plugins/Runtime/OculusRift/Source/OculusRift/Private/OculusRiftHMD.cpp b/Engine/Plugins/Runtime/OculusRift/Source/OculusRift/Private/OculusRiftHMD.cpp
--- a/Engine/Plugins/Runtime/OculusRift/Source/OculusRift/Private/OculusRiftHMD.cpp
+++ b/Engine/Plugins/Runtime/OculusRift/Source/OculusRift/Private/OculusRiftHMD.cpp
@@ -122,9 +122,9 @@ TSharedPtr<FHMDSettings, ESPMode::ThreadSafe> FOculusRiftHMD::CreateNewSettings(
 	return Result;
 }
 
-bool FOculusRiftHMD::OnStartGameFrame()
+bool FOculusRiftHMD::OnStartGameFrame( FWorldContext& WorldContext )
 {
-	bool rv = FHeadMountedDisplay::OnStartGameFrame();
+	bool rv = FHeadMountedDisplay::OnStartGameFrame( WorldContext );
 	if (!rv)
 	{
 		return false;
diff --git a/Engine/Plugins/Runtime/OculusRift/Source/OculusRift/Private/OculusRiftHMD.h b/Engine/Plugins/Runtime/OculusRift/Source/OculusRift/Private/OculusRiftHMD.h
--- a/Engine/Plugins/Runtime/OculusRift/Source/OculusRift/Private/OculusRiftHMD.h
+++ b/Engine/Plugins/Runtime/OculusRift/Source/OculusRift/Private/OculusRiftHMD.h
@@ -270,7 +270,7 @@ class FOculusRiftHMD : public FHeadMountedDisplay
 	friend class UOculusFunctionLibrary;
 public:
 	/** IHeadMountedDisplay interface */
-	virtual bool OnStartGameFrame() override;
+	virtual bool OnStartGameFrame( FWorldContext& WorldContext ) override;
 
 	virtual bool IsHMDConnected() override;
 	virtual EHMDDeviceType::Type GetHMDDeviceType() const override;
diff --git a/Engine/Source/Runtime/Engine/Private/LevelTick.cpp b/Engine/Source/Runtime/Engine/Private/LevelTick.cpp
--- a/Engine/Source/Runtime/Engine/Private/LevelTick.cpp
+++ b/Engine/Source/Runtime/Engine/Private/LevelTick.cpp
@@ -1035,7 +1035,7 @@ void UWorld::Tick( ELevelTick TickType, float DeltaSeconds )
 
 	if (GEngine->HMDDevice.IsValid())
 	{
-		GEngine->HMDDevice->OnStartGameFrame();
+		GEngine->HMDDevice->OnStartGameFrame( GEngine->GetWorldContextFromWorldChecked( this ) );
 	}
 
 #if ENABLE_COLLISION_ANALYZER
@@ -1442,7 +1442,7 @@ void UWorld::Tick( ELevelTick TickType, float DeltaSeconds )
 
 	if (GEngine->HMDDevice.IsValid())
 	{
-		GEngine->HMDDevice->OnEndGameFrame();
+		GEngine->HMDDevice->OnEndGameFrame( GEngine->GetWorldContextFromWorldChecked( this ) );
 	}
 }
 
diff --git a/Engine/Source/Runtime/HeadMountedDisplay/Public/IHeadMountedDisplay.h b/Engine/Source/Runtime/HeadMountedDisplay/Public/IHeadMountedDisplay.h
--- a/Engine/Source/Runtime/HeadMountedDisplay/Public/IHeadMountedDisplay.h
+++ b/Engine/Source/Runtime/HeadMountedDisplay/Public/IHeadMountedDisplay.h
@@ -289,12 +289,12 @@ public:
 	/**
 	 * This method is called when new game frame begins (called on a game thread).
 	 */
-	virtual bool OnStartGameFrame() { return false; }
+	virtual bool OnStartGameFrame( FWorldContext& WorldContext ) { return false; }
 
 	/**
 	 * This method is called when game frame ends (called on a game thread).
 	 */
-	virtual bool OnEndGameFrame() { return false; }
+	virtual bool OnEndGameFrame( FWorldContext& WorldContext ) { return false; }
 
 	/** 
 	 * Additional optional distorion rendering parameters
