#gearvr Re-enabling Oculus GearVR system menu

diff --git a/Engine/Plugins/Runtime/GearVR/Source/GearVR/Private/GearVR.cpp b/Engine/Plugins/Runtime/GearVR/Source/GearVR/Private/GearVR.cpp
--- a/Engine/Plugins/Runtime/GearVR/Source/GearVR/Private/GearVR.cpp
+++ b/Engine/Plugins/Runtime/GearVR/Source/GearVR/Private/GearVR.cpp
@@ -162,6 +162,15 @@ bool FGearVR::OnStartGameFrame( FWorldContext& WorldContext )
 
 	rv = GetEyePoses(*CurrentFrame, CurrentFrame->CurEyeRenderPose, CurrentFrame->CurSensorState);
 
+#if !UE_BUILD_SHIPPING
+	{ // used for debugging, do not remove
+		FQuat CurHmdOrientation;
+		FVector CurHmdPosition;
+		GetCurrentPose(CurHmdOrientation, CurHmdPosition, false, false);
+		//UE_LOG(LogHMD, Log, TEXT("BFPOSE: Pos %.3f %.3f %.3f, fr: %d"), CurHmdPosition.X, CurHmdPosition.Y, CurHmdPosition.Y,(int)CurrentFrame->FrameNumber);
+		//UE_LOG(LogHMD, Log, TEXT("BFPOSE: Yaw %.3f Pitch %.3f Roll %.3f, fr: %d"), CurHmdOrientation.Rotator().Yaw, CurHmdOrientation.Rotator().Pitch, CurHmdOrientation.Rotator().Roll, (int)CurrentFrame->FrameNumber);
+	}
+#endif
 	return rv;
 }
 
@@ -868,7 +877,7 @@ void FGearVR::StartOVRGlobalMenu()
 {
 	check(IsInRenderingThread());
 
-	//?? ovr_StartSystemActivity(&Java, PUI_GLOBAL_MENU, NULL);
+	ovr_StartSystemActivity(&JavaGT, PUI_GLOBAL_MENU, NULL);
 }
 
 void FGearVR::UpdateViewport(bool bUseSeparateRenderTarget, const FViewport& InViewport, SViewport* ViewportWidget)
