Fix Android non-unity compile failure, UE-17498.

diff --git a/Engine/Plugins/Runtime/GearVR/Source/GearVR/Private/GearVRRender.cpp b/Engine/Plugins/Runtime/GearVR/Source/GearVR/Private/GearVRRender.cpp
--- a/Engine/Plugins/Runtime/GearVR/Source/GearVR/Private/GearVRRender.cpp
+++ b/Engine/Plugins/Runtime/GearVR/Source/GearVR/Private/GearVRRender.cpp
@@ -10,6 +10,7 @@
 #include "OpenGLResources.h"
 #include "VrApi_Helpers.h"
 #include "VrApi_LocalPrefs.h"
+#include "Android/AndroidJNI.h"
 
 #define NUM_BUFFERS 3
 
