GearVR crash fix. Forwarding rendering path requires clearing both color and depth.

diff --git a/Engine/Plugins/Runtime/GearVR/Source/GearVR/Private/GearVRRender.cpp b/Engine/Plugins/Runtime/GearVR/Source/GearVR/Private/GearVRRender.cpp
--- a/Engine/Plugins/Runtime/GearVR/Source/GearVR/Private/GearVRRender.cpp
+++ b/Engine/Plugins/Runtime/GearVR/Source/GearVR/Private/GearVRRender.cpp
@@ -69,7 +69,7 @@ namespace GearVR
 			bInAllocatedStorage,
 			InFlags,
 			InTextureRange,
-			FClearValueBinding::None
+			FClearValueBinding::Black
 			)
 		{
 			FMemory::Memzero(Textures);
