GearVR rainbow rendering init fix

diff --git a/Engine/Plugins/Runtime/GearVR/Source/GearVR/Private/GearVR.cpp b/Engine/Plugins/Runtime/GearVR/Source/GearVR/Private/GearVR.cpp
--- a/Engine/Plugins/Runtime/GearVR/Source/GearVR/Private/GearVR.cpp
+++ b/Engine/Plugins/Runtime/GearVR/Source/GearVR/Private/GearVR.cpp
@@ -132,7 +132,7 @@ TSharedPtr<FHMDSettings, ESPMode::ThreadSafe> FGearVR::CreateNewSettings() const
 
 bool FGearVR::OnStartGameFrame( FWorldContext& WorldContext )
 {
-	if (GFrameNumber > 1)
+	if (GFrameNumber > 3)
 	{
 		// Temp fix to a bug in ovr_DeviceIsDocked() that can't return
 		// actual state of docking. We are switching to stereo at the start
