GearVR crash fix encountered when opening the ovr menu.

diff --git a/Engine/Plugins/Runtime/GearVR/Source/GearVR/Private/GearVR.cpp b/Engine/Plugins/Runtime/GearVR/Source/GearVR/Private/GearVR.cpp
--- a/Engine/Plugins/Runtime/GearVR/Source/GearVR/Private/GearVR.cpp
+++ b/Engine/Plugins/Runtime/GearVR/Source/GearVR/Private/GearVR.cpp
@@ -891,7 +891,10 @@ void FGearVR::StartOVRGlobalMenu()
 {
 	check(IsInRenderingThread());
 
-	ovr_StartSystemActivity(&JavaGT, PUI_GLOBAL_MENU, NULL);
+	if (pGearVRBridge)
+	{
+		ovr_StartSystemActivity(&pGearVRBridge->JavaRT, PUI_GLOBAL_MENU, NULL);
+	}
 }
 
 void FGearVR::UpdateViewport(bool bUseSeparateRenderTarget, const FViewport& InViewport, SViewport* ViewportWidget)
diff --git a/Engine/Plugins/Runtime/GearVR/Source/GearVR/Private/GearVR.h b/Engine/Plugins/Runtime/GearVR/Source/GearVR/Private/GearVR.h
--- a/Engine/Plugins/Runtime/GearVR/Source/GearVR/Private/GearVR.h
+++ b/Engine/Plugins/Runtime/GearVR/Source/GearVR/Private/GearVR.h
@@ -206,7 +206,7 @@ protected:
 class FGearVRCustomPresent : public FRHICustomPresent
 {
 	friend class FViewExtension;
-	friend class FGearVR;
+	friend class ::FGearVR;
 public:
 	FGearVRCustomPresent(jobject InActivityObject, int32 InMinimumVsyncs, int InCpuLevel, int InGpuLevel, pid_t	InGameThreadId);
 
