Speculative attempt to fix CIS.

diff --git a/Engine/Plugins/Runtime/GearVR/Source/GearVR/Private/GearVR.cpp b/Engine/Plugins/Runtime/GearVR/Source/GearVR/Private/GearVR.cpp
--- a/Engine/Plugins/Runtime/GearVR/Source/GearVR/Private/GearVR.cpp
+++ b/Engine/Plugins/Runtime/GearVR/Source/GearVR/Private/GearVR.cpp
@@ -632,7 +632,7 @@ void FGearVR::SetupViewFamily(FSceneViewFamily& InViewFamily)
 {
 	InViewFamily.EngineShowFlags.MotionBlur = 0;
 	InViewFamily.EngineShowFlags.HMDDistortion = false;
-	InViewFamily.EngineShowFlags.ScreenPercentage = false;
+	InViewFamily.EngineShowFlags.SetScreenPercentage(false);
 	InViewFamily.EngineShowFlags.StereoRendering = IsStereoEnabled();
 }
 
