Use accessor for WorldToMetersScale
#gearvr
#ue4
#android

diff --git a/Engine/Plugins/Runtime/GearVR/Source/GearVR/Private/GearVR.cpp b/Engine/Plugins/Runtime/GearVR/Source/GearVR/Private/GearVR.cpp
--- a/Engine/Plugins/Runtime/GearVR/Source/GearVR/Private/GearVR.cpp
+++ b/Engine/Plugins/Runtime/GearVR/Source/GearVR/Private/GearVR.cpp
@@ -318,6 +318,7 @@ void FGameFrame::PoseToOrientationAndPosition(const ovrPosef& InPose, FQuat& Out
 {
 	OutOrientation = ToFQuat(InPose.Orientation);
 
+	float WorldToMetersScale = GetWorldToMetersScale();
 	check(WorldToMetersScale >= 0);
 	// correct position according to BaseOrientation and BaseOffset. 
 	const FVector Pos = (ToFVector_M2U(OVR::Vector3f(InPose.Position), WorldToMetersScale) - (Settings->BaseOffset * WorldToMetersScale)) * CameraScale3D;
