Disable GearVR builds for Win32 and Win64 configurations.

diff --git a/Engine/Plugins/Runtime/GearVR/GearVR.uplugin b/Engine/Plugins/Runtime/GearVR/GearVR.uplugin
--- a/Engine/Plugins/Runtime/GearVR/GearVR.uplugin
+++ b/Engine/Plugins/Runtime/GearVR/GearVR.uplugin
@@ -22,9 +22,7 @@
 			"LoadingPhase" : "PostConfigInit",
 			"WhitelistPlatforms" :
 			[
-				"Android",
-				"Win64",
-				"Win32"
+				"Android"
 			]
 		}
 	]
