Change Gear VR installLocation to auto as now required by Oculus
#ue4
#android

diff --git a/Engine/Plugins/Runtime/GearVR/Source/GearVR/GearVR_APL.xml b/Engine/Plugins/Runtime/GearVR/Source/GearVR/GearVR_APL.xml
--- a/Engine/Plugins/Runtime/GearVR/Source/GearVR/GearVR_APL.xml
+++ b/Engine/Plugins/Runtime/GearVR/Source/GearVR/GearVR_APL.xml
@@ -67,8 +67,8 @@
 	<androidManifestUpdates>
 		<if condition="bGearDistribution">
 			<true>
-				<!-- distribution builds only install internally -->
-				<addAttribute tag="manifest" name="android:installLocation" value="internalOnly"/>
+				<!-- distribution builds can install internal or SD card -->
+				<addAttribute tag="manifest" name="android:installLocation" value="auto"/>
 
 				<!-- update the GameActivity activity -->
 				<loopElements tag="activity">
