Duplicating in main Change: 3466497 from Relase-4.16
UE-45526 Disabling Gear VR plugin causes packaged game to crash on launch
-Problem is GearVR plugin is enabled by default, and loads some content (controller mesh).
-If you unplug the plugin in a blueprint build your engine will still try to load that content, but it won't be packaged.
-Now we test if the plugin is active before loading the mesh.
-Added a comment to what was in release-4.16 because its a bit mysterious why one would do this.

diff --git a/Engine/Plugins/Runtime/GearVR/Source/GearVRController/Private/GearVRControllerComponent.cpp b/Engine/Plugins/Runtime/GearVR/Source/GearVRController/Private/GearVRControllerComponent.cpp
--- a/Engine/Plugins/Runtime/GearVR/Source/GearVRController/Private/GearVRControllerComponent.cpp
+++ b/Engine/Plugins/Runtime/GearVR/Source/GearVRController/Private/GearVRControllerComponent.cpp
@@ -5,6 +5,7 @@
 #include "Components/StaticMeshComponent.h"
 #include "Engine/StaticMesh.h"
 #include "MotionControllerComponent.h"
+#include "ModuleManager.h"
 
 UGearVRControllerComponent::UGearVRControllerComponent()
 : ControllerMesh(nullptr)
@@ -14,8 +15,14 @@ UGearVRControllerComponent::UGearVRControllerComponent()
 	PrimaryComponentTick.bCanEverTick = true;
 	bAutoActivate = true;
 
-
-	ControllerMesh = Cast<UStaticMesh>(StaticLoadObject(UStaticMesh::StaticClass(), NULL, TEXT("/GearVR/Meshes/GearVRController")));
+	// GearVR plugin is enabled by default, so it is built into the executable for content projects.
+	// If you disable the plugin the mesh content will be missing, and when the default object is 
+	// constructed you get a fatal load error.
+	// This check avoids that error.
+	if (FModuleManager::Get().IsModuleLoaded("GearVR"))
+	{
+		ControllerMesh = Cast<UStaticMesh>(StaticLoadObject(UStaticMesh::StaticClass(), NULL, TEXT("/GearVR/Meshes/GearVRController")));
+	}
 }
 
 UMotionControllerComponent* UGearVRControllerComponent::GetMotionController() const
